Disabling GoogleVR support for Mac until warnings are fixed.

diff --git a/Engine/Plugins/Runtime/GoogleVR/GoogleVRController/GoogleVRController.uplugin b/Engine/Plugins/Runtime/GoogleVR/GoogleVRController/GoogleVRController.uplugin
--- a/Engine/Plugins/Runtime/GoogleVR/GoogleVRController/GoogleVRController.uplugin
+++ b/Engine/Plugins/Runtime/GoogleVR/GoogleVRController/GoogleVRController.uplugin
@@ -20,7 +20,7 @@
             "Name": "GoogleVRController",
             "Type": "Runtime",
             "LoadingPhase": "PostConfigInit",
-            "WhitelistPlatforms": [ "Android", "Mac", "Win32", "Win64"]
+            "WhitelistPlatforms": [ "Android", "Win32", "Win64"]
         }
     ]
 }
diff --git a/Engine/Plugins/Runtime/GoogleVR/GoogleVRHMD/GoogleVRHMD.uplugin b/Engine/Plugins/Runtime/GoogleVR/GoogleVRHMD/GoogleVRHMD.uplugin
--- a/Engine/Plugins/Runtime/GoogleVR/GoogleVRHMD/GoogleVRHMD.uplugin
+++ b/Engine/Plugins/Runtime/GoogleVR/GoogleVRHMD/GoogleVRHMD.uplugin
@@ -20,7 +20,7 @@
             "Name": "GoogleVRHMD",
             "Type": "Runtime",
             "LoadingPhase": "PostConfigInit",
-            "WhitelistPlatforms": [ "Win64", "Win32", "Mac", "Android", "IOS" ]
+            "WhitelistPlatforms": [ "Win64", "Win32", "Android", "IOS" ]
         }
     ]
 }
