Compiler fix for GoogleVRController

diff --git a/Engine/Plugins/Runtime/GoogleVR/GoogleVRController/Source/GoogleVRController/Private/GoogleVRController.cpp b/Engine/Plugins/Runtime/GoogleVR/GoogleVRController/Source/GoogleVRController/Private/GoogleVRController.cpp
--- a/Engine/Plugins/Runtime/GoogleVR/GoogleVRController/Source/GoogleVRController/Private/GoogleVRController.cpp
+++ b/Engine/Plugins/Runtime/GoogleVR/GoogleVRController/Source/GoogleVRController/Private/GoogleVRController.cpp
@@ -29,7 +29,7 @@ class FGoogleVRControllerPlugin : public IGoogleVRControllerPlugin
 	gvr::ControllerApi* CreateAndInitGoogleVRControllerAPI()
 	{
 		// Get controller API
-		gvr::ControllerApi* pController = new gvr::ControllerApi();
+		ControllerApiPtr pController = new gvr::ControllerApi();
 		check(pController);
 
 		gvr::ControllerApiOptions options;
@@ -74,7 +74,7 @@ class FGoogleVRControllerPlugin : public IGoogleVRControllerPlugin
 	{
 #if GOOGLEVRCONTROLLER_SUPPORTED_PLATFORMS
 		UE_LOG(LogTemp, Warning, TEXT("Creating Input Device: GoogleVRController -- Supported"));
-		gvr::ControllerApi* pControllerAPI = CreateAndInitGoogleVRControllerAPI();
+		ControllerApi* pControllerAPI = CreateAndInitGoogleVRControllerAPI();
 		if(pControllerAPI)
 		{
 			return TSharedPtr< class IInputDevice >(new FGoogleVRController(pControllerAPI, InMessageHandler));
@@ -92,13 +92,13 @@ class FGoogleVRControllerPlugin : public IGoogleVRControllerPlugin
 
 IMPLEMENT_MODULE( FGoogleVRControllerPlugin, GoogleVRController)
 
-FGoogleVRController::FGoogleVRController(gvr::ControllerApi* pControllerAPI, const TSharedRef< FGenericApplicationMessageHandler >& InMessageHandler)
+FGoogleVRController::FGoogleVRController(ControllerApiPtr pControllerAPI, const TSharedRef< FGenericApplicationMessageHandler >& InMessageHandler)
 	:
 #if GOOGLEVRCONTROLLER_SUPPORTED_PLATFORMS
-	pController(pControllerAPI)
+	pController(pControllerAPI),
 #endif
-	,bControllerReadyToPollState(false)
-	,MessageHandler(InMessageHandler)
+	bControllerReadyToPollState(false),
+	MessageHandler(InMessageHandler)
 {
 	UE_LOG(LogTemp, Warning, TEXT("GoogleVR Controller Created"));
 
diff --git a/Engine/Plugins/Runtime/GoogleVR/GoogleVRController/Source/GoogleVRController/Private/GoogleVRController.h b/Engine/Plugins/Runtime/GoogleVR/GoogleVRController/Source/GoogleVRController/Private/GoogleVRController.h
--- a/Engine/Plugins/Runtime/GoogleVR/GoogleVRController/Source/GoogleVRController/Private/GoogleVRController.h
+++ b/Engine/Plugins/Runtime/GoogleVR/GoogleVRController/Source/GoogleVRController/Private/GoogleVRController.h
@@ -21,6 +21,9 @@
 
 #if GOOGLEVRCONTROLLER_SUPPORTED_PLATFORMS
 using namespace gvr;
+typedef gvr::ControllerApi* ControllerApiPtr;
+#else
+typedef void* ControllerApiPtr;
 #endif
 
 DEFINE_LOG_CATEGORY_STATIC(LogGoogleVRController, Log, All);
@@ -48,7 +51,7 @@ class FGoogleVRController : public IInputDevice, public IMotionController
 {
 public:
 
-	FGoogleVRController(gvr::ControllerApi* pControllerAPI, const TSharedRef< FGenericApplicationMessageHandler >& InMessageHandler);
+	FGoogleVRController(ControllerApiPtr pControllerAPI, const TSharedRef< FGenericApplicationMessageHandler >& InMessageHandler);
 	virtual ~FGoogleVRController();
 
 public:
@@ -133,7 +136,7 @@ private:
 
 #if GOOGLEVRCONTROLLER_SUPPORTED_PLATFORMS
 	/** GVR Controller client reference */
-	gvr::ControllerApi *pController;
+	ControllerApiPtr pController;
 
 	/** Cached controller info */
 	gvr::ControllerState CachedControllerState;
