Also disable GoogleVRController for Mac, which depends on the GoogleVR plugin.

diff --git a/Engine/Plugins/Runtime/GoogleVR/GoogleVRController/GoogleVRController.uplugin b/Engine/Plugins/Runtime/GoogleVR/GoogleVRController/GoogleVRController.uplugin
--- a/Engine/Plugins/Runtime/GoogleVR/GoogleVRController/GoogleVRController.uplugin
+++ b/Engine/Plugins/Runtime/GoogleVR/GoogleVRController/GoogleVRController.uplugin
@@ -20,7 +20,7 @@
             "Name": "GoogleVRController",
             "Type": "Runtime",
             "LoadingPhase": "PostConfigInit",
-            "WhitelistPlatforms": [ "Android", "Mac", "Win32", "Win64"]
+            "WhitelistPlatforms": [ "Android", "Win32", "Win64"]
         }
     ]
 }
