UE-43768 //UE4/Main: Step 'Compile UE4Game Mac' - 1 Error
-Removed mac from GoogleVRController plugin whitelist. That plugin depends on GoogleVR plugin, from who's whitelist mac was previously removed.

diff --git a/Engine/Plugins/Runtime/GoogleVR/GoogleVRController/GoogleVRController.uplugin b/Engine/Plugins/Runtime/GoogleVR/GoogleVRController/GoogleVRController.uplugin
--- a/Engine/Plugins/Runtime/GoogleVR/GoogleVRController/GoogleVRController.uplugin
+++ b/Engine/Plugins/Runtime/GoogleVR/GoogleVRController/GoogleVRController.uplugin
@@ -20,7 +20,7 @@
             "Name": "GoogleVRController",
             "Type": "Runtime",
             "LoadingPhase": "PostConfigInit",
-            "WhitelistPlatforms": [ "Android", "Mac", "Win32", "Win64", "IOS" ]
+            "WhitelistPlatforms": [ "Android", "Win32", "Win64", "IOS" ]
         }
     ]
 }
