Use TInlineComponentArray if you're going to use GetComponens that takes an array
#rnx

diff --git a/Engine/Plugins/Runtime/GoogleVR/GoogleVRController/Source/GoogleVRController/Private/GoogleVRControllerTooltipComponent.cpp b/Engine/Plugins/Runtime/GoogleVR/GoogleVRController/Source/GoogleVRController/Private/GoogleVRControllerTooltipComponent.cpp
--- a/Engine/Plugins/Runtime/GoogleVR/GoogleVRController/Source/GoogleVRController/Private/GoogleVRControllerTooltipComponent.cpp
+++ b/Engine/Plugins/Runtime/GoogleVR/GoogleVRController/Source/GoogleVRController/Private/GoogleVRControllerTooltipComponent.cpp
@@ -44,7 +44,7 @@ void UGoogleVRControllerTooltipComponent::BeginPlay()
 	}
 
 	// Attempt to attach this tooltip to the motion controller and set it's position.
-	TArray<UMotionControllerComponent*> Components;
+	TInlineComponentArray<UMotionControllerComponent*> Components;
 	GetOwner()->GetComponents(Components);
 	if(Components.Num() == 0)
 	{
diff --git a/Engine/Plugins/Runtime/GoogleVR/GoogleVRController/Source/GoogleVRController/Private/GoogleVRGazeReticleComponent.cpp b/Engine/Plugins/Runtime/GoogleVR/GoogleVRController/Source/GoogleVRController/Private/GoogleVRGazeReticleComponent.cpp
--- a/Engine/Plugins/Runtime/GoogleVR/GoogleVRController/Source/GoogleVRController/Private/GoogleVRGazeReticleComponent.cpp
+++ b/Engine/Plugins/Runtime/GoogleVR/GoogleVRController/Source/GoogleVRController/Private/GoogleVRGazeReticleComponent.cpp
@@ -64,7 +64,7 @@ void UGoogleVRGazeReticleComponent::BeginPlay()
 {
 	Super::BeginPlay();
 
-	TArray<UGoogleVRPointerInputComponent*> Components;
+	TInlineComponentArray<UGoogleVRPointerInputComponent*> Components;
 	UGoogleVRPointerInputComponent* InputComponent = nullptr;
 
 	GetOwner()->GetComponents(Components);
@@ -91,7 +91,7 @@ void UGoogleVRGazeReticleComponent::BeginPlay()
 		InputComponent->SetPointer(this);
 	}
 
-	TArray<UCameraComponent*> CameraComponents;
+	TInlineComponentArray<UCameraComponent*> CameraComponents;
 	GetOwner()->GetComponents(CameraComponents);
 	if (CameraComponents.Num() == 0)
 	{
diff --git a/Engine/Plugins/Runtime/GoogleVR/GoogleVRController/Source/GoogleVRController/Private/GoogleVRMotionControllerComponent.cpp b/Engine/Plugins/Runtime/GoogleVR/GoogleVRController/Source/GoogleVRController/Private/GoogleVRMotionControllerComponent.cpp
--- a/Engine/Plugins/Runtime/GoogleVR/GoogleVRController/Source/GoogleVRController/Private/GoogleVRMotionControllerComponent.cpp
+++ b/Engine/Plugins/Runtime/GoogleVR/GoogleVRController/Source/GoogleVRController/Private/GoogleVRMotionControllerComponent.cpp
@@ -217,7 +217,7 @@ void UGoogleVRMotionControllerComponent::BeginPlay()
 {
 	Super::BeginPlay();
 
-	TArray<UGoogleVRPointerInputComponent*> PointerInputComponents;
+	TInlineComponentArray<UGoogleVRPointerInputComponent*> PointerInputComponents;
 	GetOwner()->GetComponents(PointerInputComponents);
 
 	if (PointerInputComponents.Num() == 0)
