Fix lots of CIS errors introduced by CL 7099376. Some static symbols are duplicated in translation units, which causes duplicate symbol errors in unity mode.

#rnx

diff --git a/Engine/Plugins/Runtime/GoogleVR/GoogleVRController/Source/GoogleVRController/Private/GoogleVRController.cpp b/Engine/Plugins/Runtime/GoogleVR/GoogleVRController/Source/GoogleVRController/Private/GoogleVRController.cpp
--- a/Engine/Plugins/Runtime/GoogleVR/GoogleVRController/Source/GoogleVRController/Private/GoogleVRController.cpp
+++ b/Engine/Plugins/Runtime/GoogleVR/GoogleVRController/Source/GoogleVRController/Private/GoogleVRController.cpp
@@ -30,6 +30,8 @@ DEFINE_LOG_CATEGORY_STATIC(LogGoogleVRController, Log, All);
 #include "GoogleVRInstantPreviewGetServer.h"
 #endif  // GOOGLEVRCONTROLLER_SUPPORTED_INSTANT_PREVIEW_PLATFORMS
 
+DEFINE_LOG_CATEGORY(LogGoogleVRMotionController);
+
 #if GOOGLEVRCONTROLLER_SUPPORTED_EMULATOR_PLATFORMS
 #define CONTROLLER_EVENT_FORWARDED_PORT 7003 //Change this port number if it is already taken.
 #define ADB_FORWARD_RETRY_TIME 5.0 //5 seconds
diff --git a/Engine/Plugins/Runtime/GoogleVR/GoogleVRController/Source/GoogleVRController/Private/GoogleVRController.h b/Engine/Plugins/Runtime/GoogleVR/GoogleVRController/Source/GoogleVRController/Private/GoogleVRController.h
--- a/Engine/Plugins/Runtime/GoogleVR/GoogleVRController/Source/GoogleVRController/Private/GoogleVRController.h
+++ b/Engine/Plugins/Runtime/GoogleVR/GoogleVRController/Source/GoogleVRController/Private/GoogleVRController.h
@@ -23,6 +23,8 @@
 using namespace gvr;
 #endif
 
+DECLARE_LOG_CATEGORY_EXTERN(LogGoogleVRMotionController, Log, All);
+
 /** Total number of controllers in a set */
 #define CONTROLLERS_PER_PLAYER	2
 
diff --git a/Engine/Plugins/Runtime/GoogleVR/GoogleVRController/Source/GoogleVRController/Private/GoogleVRLaserVisualComponent.cpp b/Engine/Plugins/Runtime/GoogleVR/GoogleVRController/Source/GoogleVRController/Private/GoogleVRLaserVisualComponent.cpp
--- a/Engine/Plugins/Runtime/GoogleVR/GoogleVRController/Source/GoogleVRController/Private/GoogleVRLaserVisualComponent.cpp
+++ b/Engine/Plugins/Runtime/GoogleVR/GoogleVRController/Source/GoogleVRController/Private/GoogleVRLaserVisualComponent.cpp
@@ -18,8 +18,6 @@
 #include "Modules/ModuleManager.h"
 #include "Engine/StaticMesh.h"
 
-DEFINE_LOG_CATEGORY_STATIC(LogGoogleVRMotionController, Log, All);
-
 // Sets default values for this component's properties
 UGoogleVRLaserVisualComponent::UGoogleVRLaserVisualComponent()
 : LaserPlaneMesh(nullptr)
diff --git a/Engine/Plugins/Runtime/GoogleVR/GoogleVRController/Source/GoogleVRController/Private/GoogleVRMotionControllerComponent.cpp b/Engine/Plugins/Runtime/GoogleVR/GoogleVRController/Source/GoogleVRController/Private/GoogleVRMotionControllerComponent.cpp
--- a/Engine/Plugins/Runtime/GoogleVR/GoogleVRController/Source/GoogleVRController/Private/GoogleVRMotionControllerComponent.cpp
+++ b/Engine/Plugins/Runtime/GoogleVR/GoogleVRController/Source/GoogleVRController/Private/GoogleVRMotionControllerComponent.cpp
@@ -19,8 +19,6 @@
 #include "Engine/StaticMesh.h"
 #include "Modules/ModuleManager.h"
 
-DEFINE_LOG_CATEGORY_STATIC(LogGoogleVRMotionController, Log, All);
-
 const FVector UGoogleVRMotionControllerComponent::TOUCHPAD_POINT_DIMENSIONS = FVector(0.01f, 0.01f, 0.0004f);
 const FVector UGoogleVRMotionControllerComponent::BATTERY_INDICATOR_TRANSLATION = FVector(-3.0f, 0.0f, 0.001f);
 const FVector UGoogleVRMotionControllerComponent::BATTERY_INDICATOR_SCALE = FVector(0.032f, 0.015f, 1.0f);
