Merging using Dev-VR_to_Release-4.21

4345353 UE-62255 Google 6DOF Motion Controller Support
-Added support for 6DOF controller tracking on daydream.
-Project Settings->Android->Advanced APK Packaging->Configure GoogleVR to support specific hardware configurations-> Set to Daydream(6.6 DoF) to enable support for two 6dof controllers. A single 3dof controller is also supported by this configuration, as is a single 6dof controller.
-Currently the 3dof controller is detected only by noticing that the positions returned for all controllers are exactly 0,0,0 per Google reccomendation.  A future API should make it explicit when a controller is 3dof or 6dof.
-Currently we arbitrarily make the controller at index 0 the right and the one at 1 the left.  A future api should let us query for controller handedness.

4346387 Fixing build break in GoogleVRControllerFunctionLibrary.h

4349142 Fixing GoogleVRController build warning

4349742 Fixing AndroidRuntimeSettings build errors.

4350977 Fix for static analysis warning

4351170 Build fix for bad assert

diff --git a/Engine/Plugins/Runtime/GoogleVR/GoogleVRController/Source/GoogleVRController/Classes/GoogleVRControllerFunctionLibrary.h b/Engine/Plugins/Runtime/GoogleVR/GoogleVRController/Source/GoogleVRController/Classes/GoogleVRControllerFunctionLibrary.h
--- a/Engine/Plugins/Runtime/GoogleVR/GoogleVRController/Source/GoogleVRController/Classes/GoogleVRControllerFunctionLibrary.h
+++ b/Engine/Plugins/Runtime/GoogleVR/GoogleVRController/Source/GoogleVRController/Classes/GoogleVRControllerFunctionLibrary.h
@@ -4,6 +4,7 @@
 
 #include "GoogleVRControllerEventManager.h"
 #include "Kismet/BlueprintFunctionLibrary.h"
+#include "InputCoreTypes.h"
 #include "GoogleVRControllerFunctionLibrary.generated.h"
 
 UENUM(BlueprintType)
@@ -80,7 +81,7 @@ public:
 	 * Get the GoogleVR Controller state
 	 */
 	UFUNCTION(BlueprintPure, Category = "GoogleVRController", meta = (Keywords = "Cardboard AVR GVR"))
-	static EGoogleVRControllerState GetGoogleVRControllerState();
+	static EGoogleVRControllerState GetGoogleVRControllerState(EControllerHand Hand = EControllerHand::AnyHand);
 
 	/**
 	 * Get user's handedness preference from GVRSDK
@@ -93,19 +94,19 @@ public:
 	 * This function return the controller acceleration in gvr controller space.
 	 */
 	UFUNCTION(BlueprintPure, Category = "GoogleVRController", meta = (Keywords = "Cardboard AVR GVR"))
-	static FVector GetGoogleVRControllerRawAccel();
+	static FVector GetGoogleVRControllerRawAccel(EControllerHand Hand = EControllerHand::AnyHand);
 
 	/**
 	 * This function return the controller angular velocity about each axis (positive means clockwise when sighting along axis) in gvr controller space.
 	 */
 	UFUNCTION(BlueprintPure, Category = "GoogleVRController", meta = (Keywords = "Cardboard AVR GVR"))
-	static FVector GetGoogleVRControllerRawGyro();
+	static FVector GetGoogleVRControllerRawGyro(EControllerHand Hand = EControllerHand::AnyHand);
 
 	/**
 	 * This function return the orientation of the controller in unreal space.
 	 */
 	UFUNCTION(BlueprintPure, Category = "GoogleVRController", meta = (Keywords = "Cardboard AVR GVR"))
-	static FRotator GetGoogleVRControllerOrientation();
+	static FRotator GetGoogleVRControllerOrientation(EControllerHand Hand = EControllerHand::AnyHand);
 
 	/**
 	 * Return a pointer to the UGoogleVRControllerEventManager to hook up GoogleVR Controller specific event.
@@ -292,18 +293,18 @@ public:
 	 *  @return true if the battery is charging.
 	 */
 	UFUNCTION(BlueprintPure, Category = "GoogleVRController", meta = (Keywords = "Cardboard AVR GVR"))
-	static bool GetBatteryCharging();
+	static bool GetBatteryCharging(EControllerHand Hand = EControllerHand::AnyHand);
 
 	/** Get the bucketed controller battery level.
 	 *  Note this is an approximate level described by enumeration, not a percent.
 	 *  @return the approximate battery level, or unknown if the level can not be determined.
 	 */
 	UFUNCTION(BlueprintPure, Category = "GoogleVRController", meta = (Keywords = "Cardboard AVR GVR"))
-	static EGoogleVRControllerBatteryLevel GetBatteryLevel();
+	static EGoogleVRControllerBatteryLevel GetBatteryLevel(EControllerHand Hand = EControllerHand::AnyHand);
 
 	/** Get the timestamp (nanos) when the last battery event was received.
 	 *  @return the timestamp, or zero if unavailable.
 	 */
-	static int64_t GetLastBatteryTimestamp();
+	static int64_t GetLastBatteryTimestamp(EControllerHand Hand = EControllerHand::AnyHand);
 
 };
diff --git a/Engine/Plugins/Runtime/GoogleVR/GoogleVRController/Source/GoogleVRController/GoogleVRController.Build.cs b/Engine/Plugins/Runtime/GoogleVR/GoogleVRController/Source/GoogleVRController/GoogleVRController.Build.cs
--- a/Engine/Plugins/Runtime/GoogleVR/GoogleVRController/Source/GoogleVRController/GoogleVRController.Build.cs
+++ b/Engine/Plugins/Runtime/GoogleVR/GoogleVRController/Source/GoogleVRController/GoogleVRController.Build.cs
@@ -34,13 +34,19 @@ namespace UnrealBuildTool.Rules
 					"GoogleVRHMD",
 					"UMG",
 					"Slate",
-					"SlateCore"
+					"SlateCore",
 				}
 				);
 
 			if (Target.Platform == UnrealTargetPlatform.Android)
 			{
-				PrivateDependencyModuleNames.Add("Launch");
+				PrivateDependencyModuleNames.AddRange(
+					new string[]
+					{
+						"Launch",
+						"AndroidRuntimeSettings",
+					}
+					);
 			}
 
 			if (Target.bBuildEditor == true)
diff --git a/Engine/Plugins/Runtime/GoogleVR/GoogleVRController/Source/GoogleVRController/GoogleVRController_APL.xml b/Engine/Plugins/Runtime/GoogleVR/GoogleVRController/Source/GoogleVRController/GoogleVRController_APL.xml
--- a/Engine/Plugins/Runtime/GoogleVR/GoogleVRController/Source/GoogleVRController/GoogleVRController_APL.xml
+++ b/Engine/Plugins/Runtime/GoogleVR/GoogleVRController/Source/GoogleVRController/GoogleVRController_APL.xml
@@ -41,6 +41,9 @@
 	<!-- optional updates applied to AndroidManifest.xml -->
 	<androidManifestUpdates>
     <addPermission android:name="android.permission.INTERNET"/>
+      <addElements tag="application">
+        <meta-data android:name="com.google.vr.controllers.supportedConfigs" android:value="3dof|2x6dof"/>
+      </addElements>
 	</androidManifestUpdates>
 
   <gameActivityImportAdditions>
diff --git a/Engine/Plugins/Runtime/GoogleVR/GoogleVRController/Source/GoogleVRController/Private/GoogleVRController.cpp b/Engine/Plugins/Runtime/GoogleVR/GoogleVRController/Source/GoogleVRController/Private/GoogleVRController.cpp
--- a/Engine/Plugins/Runtime/GoogleVR/GoogleVRController/Source/GoogleVRController/Private/GoogleVRController.cpp
+++ b/Engine/Plugins/Runtime/GoogleVR/GoogleVRController/Source/GoogleVRController/Private/GoogleVRController.cpp
@@ -11,6 +11,7 @@
 #include "GameFramework/WorldSettings.h"
 #include "Engine/World.h"
 #include "Engine/Engine.h"
+#include "Misc/ConfigCacheIni.h"
 
 #include "GoogleVRAdbUtils.h"
 
@@ -64,6 +65,7 @@ class FGoogleVRControllerPlugin : public IGoogleVRControllerPlugin
 		options |= GVR_CONTROLLER_ENABLE_GYRO;
 		options |= GVR_CONTROLLER_ENABLE_TOUCH;
 		options |= GVR_CONTROLLER_ENABLE_ORIENTATION;
+		options |= GVR_CONTROLLER_ENABLE_POSITION;
 
 		bool success = false;
 #if GOOGLEVRCONTROLLER_SUPPORTED_ANDROID_PLATFORMS
@@ -129,7 +131,6 @@ FGoogleVRController::FGoogleVRController(const TSharedRef< FGenericApplicationMe
 	: MessageHandler(InMessageHandler)
 #endif
 	, bUseArmModel(true)
-	, CurrentControllerState(EGoogleVRControllerState::Disconnected)
 #if GOOGLEVRCONTROLLER_SUPPORTED_EMULATOR_PLATFORMS
 	, BaseEmulatorOrientation(FRotator::ZeroRotator)
 #endif
@@ -137,7 +138,27 @@ FGoogleVRController::FGoogleVRController(const TSharedRef< FGenericApplicationMe
 	, InstantPreviewControllerState()
 #endif
 {
-	UE_LOG(LogGoogleVRController, Log, TEXT("GoogleVR Controller Created"));
+	for (int32 i = 0; i < CONTROLLERS_PER_PLAYER; i++)
+	{
+		CurrentControllerStates[i] = EGoogleVRControllerState::Disconnected;
+	}
+
+#if GOOGLEVRCONTROLLER_SUPPORTED_ANDROID_PLATFORMS
+	{
+		GoogleVRCaps = EGoogleVRCaps::Cardboard;
+		FString ValueString;
+		GConfig->GetString(TEXT("/Script/AndroidRuntimeSettings.AndroidRuntimeSettings"), TEXT("GoogleVRCaps"), ValueString, GEngineIni);
+		UEnum* Enum = FindObject<UEnum>((UObject*)ANY_PACKAGE, TEXT("EGoogleVRCaps"), true);
+		if (Enum)
+		{
+			int64 Value = Enum->GetValueByName(FName(*ValueString));
+			if (Value != INDEX_NONE)
+			{
+				GoogleVRCaps = static_cast<EGoogleVRCaps::Type>(Value);
+			}
+		}
+	}
+#endif
 
 #if GOOGLEVRCONTROLLER_SUPPORTED_PLATFORMS
 	// Register motion controller!
@@ -199,6 +220,7 @@ FGoogleVRController::FGoogleVRController(const TSharedRef< FGenericApplicationMe
 	Buttons[(int32)EControllerHand::Left][EGoogleVRControllerButton::TouchPadPress] = FGamepadKeyNames::MotionController_Left_Thumbstick;
 	Buttons[(int32)EControllerHand::Right][EGoogleVRControllerButton::TouchPadPress] = FGamepadKeyNames::MotionController_Right_Thumbstick;
 
+	Buttons[(int32)EControllerHand::Left][EGoogleVRControllerButton::TouchPadTouch] = GoogleVRControllerKeyNames::Touch0;
 	Buttons[(int32)EControllerHand::Right][EGoogleVRControllerButton::TouchPadTouch] = GoogleVRControllerKeyNames::Touch0;
 
 	// Register callbacks for pause and resume
@@ -243,6 +265,7 @@ void FGoogleVRController::ApplicationResumeDelegate()
 void FGoogleVRController::PollController(float DeltaTime)
 {
 #if GOOGLEVRCONTROLLER_SUPPORTED_PLATFORMS
+
 #if GOOGLEVRCONTROLLER_SUPPORTED_INSTANT_PREVIEW_PLATFORMS
 	instant_preview::Session *session = ip_static_server_acquire_active_session(IpServerHandle);
 	if (NULL != session) {
@@ -250,37 +273,48 @@ void FGoogleVRController::PollController(float DeltaTime)
 	}
 	ip_static_server_release_active_session(IpServerHandle, session);
 #endif
+
 #if GOOGLEVRCONTROLLER_SUPPORTED_EMULATOR_PLATFORMS
-	int32_t PreviousConnectionState = CachedControllerState.GetConnectionState();
-	// If controller connection is requested but it is not connected, try resetup adb forward.
-	if(bKeepConnectingControllerEmulator
-	   && PreviousConnectionState != gvr::ControllerConnectionState::GVR_CONTROLLER_CONNECTED)
 	{
-		double CurrentTime = FPlatformTime::Seconds();
-		if(CurrentTime - LastTimeTryAdbForward > ADB_FORWARD_RETRY_TIME)
+		gvr::ControllerState* ControllerState = GetCachedControllerState(EControllerHand::AnyHand);
+		check(ControllerState);
+		int32_t PreviousConnectionState = ControllerState->GetConnectionState();
+		// If controller connection is requested but it is not connected, try resetup adb forward.
+		if (bKeepConnectingControllerEmulator
+			&& PreviousConnectionState != gvr::ControllerConnectionState::GVR_CONTROLLER_CONNECTED)
 		{
-			UE_LOG(LogGoogleVRController, Log, TEXT("Trying to connect to GoogleVR Controller"));
-			SetupAdbForward();
-			LastTimeTryAdbForward = CurrentTime;
+			double CurrentTime = FPlatformTime::Seconds();
+			if (CurrentTime - LastTimeTryAdbForward > ADB_FORWARD_RETRY_TIME)
+			{
+				UE_LOG(LogGoogleVRController, Log, TEXT("Trying to connect to GoogleVR Controller"));
+				SetupAdbForward();
+				LastTimeTryAdbForward = CurrentTime;
+			}
 		}
-	}
 
-	CachedControllerState.Update(*pController);
+		ControllerState->Update(*pController);
 
-	if(PreviousConnectionState != gvr::ControllerConnectionState::GVR_CONTROLLER_CONNECTED
-	   && CachedControllerState.GetConnectionState() == gvr::ControllerConnectionState::GVR_CONTROLLER_CONNECTED)
-	{
-		UE_LOG(LogGoogleVRController, Log, TEXT("GoogleVR Controller Connected"));
-	}
+		if (PreviousConnectionState != gvr::ControllerConnectionState::GVR_CONTROLLER_CONNECTED
+			&& ControllerState->GetConnectionState() == gvr::ControllerConnectionState::GVR_CONTROLLER_CONNECTED)
+		{
+			UE_LOG(LogGoogleVRController, Log, TEXT("GoogleVR Controller Connected"));
+		}
 
-	if(PreviousConnectionState == gvr::ControllerConnectionState::GVR_CONTROLLER_CONNECTED
-	   && CachedControllerState.GetConnectionState() != gvr::ControllerConnectionState::GVR_CONTROLLER_CONNECTED)
-	{
-		UE_LOG(LogGoogleVRController, Log, TEXT("GoogleVR Controller Disconnected"));
+		if (PreviousConnectionState == gvr::ControllerConnectionState::GVR_CONTROLLER_CONNECTED
+			&& ControllerState->GetConnectionState() != gvr::ControllerConnectionState::GVR_CONTROLLER_CONNECTED)
+		{
+			UE_LOG(LogGoogleVRController, Log, TEXT("GoogleVR Controller Disconnected"));
+		}
 	}
 #elif GOOGLEVRCONTROLLER_SUPPORTED_ANDROID_PLATFORMS
-	CachedControllerState.Update(*pController);
+	for (int32 i = 0; i < CONTROLLERS_PER_PLAYER; i++)
+	{
+		CachedControllerStates[i].Update(*pController, i);
+	}
+
 #endif
+
+	// In a future google api hopefully we will have a better way to detect 3dof controllers, and we might want to avoid doing this work for 6dof controllers (where it will not be used).
 	if (bUseArmModel)
 	{
 		// Update the handedness. This could be changed in UserSettings at anytime so we poll for it.
@@ -316,20 +350,23 @@ void FGoogleVRController::PollController(float DeltaTime)
 		else
 #endif  // GOOGLEVRCONTROLLER_SUPPORTED_INSTANT_PREVIEW_PLATFORMS
 		{
+			gvr::ControllerState* ControllerState = GetCachedControllerState(EControllerHand::AnyHand);
+			check(ControllerState);
+
 			// get acceleration data
-			gvr_vec3f GvrAccel = CachedControllerState.GetAccel();
+			gvr_vec3f GvrAccel = ControllerState->GetAccel();
 			UpdateData.acceleration = gvr_arm_model::Vector3(GvrAccel.x, GvrAccel.y, GvrAccel.z);
 
 			// Get orientation data
-			gvr_quatf GvrOrientation = CachedControllerState.GetOrientation();
+			gvr_quatf GvrOrientation = ControllerState->GetOrientation();
 			UpdateData.orientation = gvr_arm_model::Quaternion(GvrOrientation.qw, GvrOrientation.qx, GvrOrientation.qy, GvrOrientation.qz);
 
 			// Get gyroscope data
-			gvr_vec3f GvrGyro = CachedControllerState.GetGyro();
+			gvr_vec3f GvrGyro = ControllerState->GetGyro();
 			UpdateData.gyro = gvr_arm_model::Vector3(GvrGyro.x, GvrGyro.y, GvrGyro.z);
 
 			// Get connected status
-			UpdateData.connected = CachedControllerState.GetConnectionState() == gvr::ControllerConnectionState::GVR_CONTROLLER_CONNECTED;
+			UpdateData.connected = ControllerState->GetConnectionState() == gvr::ControllerConnectionState::GVR_CONTROLLER_CONNECTED;
 		}
 
 		// Get head direction and position of the HMD, used for FollowGaze options
@@ -363,16 +400,19 @@ void FGoogleVRController::PollController(float DeltaTime)
 #endif
 }
 
-void FGoogleVRController::ProcessControllerButtons()
+void FGoogleVRController::ProcessControllerButtons(int32 ControllerStateIndex)
 {
 #if GOOGLEVRCONTROLLER_SUPPORTED_PLATFORMS
+	const EControllerHand Hand = GetControllerHandFromStateIndex(ControllerStateIndex);
+	check(Hand == EControllerHand::Left || Hand == EControllerHand::Right);
+
 	// Capture our current button states
 	bool CurrentButtonStates[EGoogleVRControllerButton::TotalButtonCount] = {0};
 
 	FVector2D TranslatedLocation = FVector2D::ZeroVector;
 
 #if GOOGLEVRCONTROLLER_SUPPORTED_INSTANT_PREVIEW_PLATFORMS
-	if (gvr::ControllerConnectionState::GVR_CONTROLLER_CONNECTED == InstantPreviewControllerState.connection_state)
+	if (ControllerStateIndex == 0 && gvr::ControllerConnectionState::GVR_CONTROLLER_CONNECTED == InstantPreviewControllerState.connection_state)
 	{
 		// Process our known set of buttons
 		CurrentButtonStates[EGoogleVRControllerButton::TouchPadPress] = InstantPreviewControllerState.click_button_state;
@@ -402,8 +442,10 @@ void FGoogleVRController::ProcessControllerButtons()
 	else
 #endif  // GOOGLEVRCONTROLLER_SUPPORTED_INSTANT_PREVIEW_PLATFORMS
 	{
-		if (IsAvailable())
+		if (IsAvailable(ControllerStateIndex))
 		{
+			const gvr::ControllerState& CachedControllerState = CachedControllerStates[ControllerStateIndex];
+
 			// Process our known set of buttons
 			if (CachedControllerState.GetButtonState(gvr::ControllerButton::GVR_CONTROLLER_BUTTON_CLICK))
 			{
@@ -434,9 +476,23 @@ void FGoogleVRController::ProcessControllerButtons()
 				CurrentButtonStates[EGoogleVRControllerButton::ApplicationMenu] = false;
 			}
 
-			// Note: There is no Grip or Trigger button information from the CachedControllerState; so do nothing
-			// EGoogleVRControllerButton::TriggerPress - unhandled
-			// EGoogleVRControllerButton::Grip - unhandled
+			if (CachedControllerState.GetButtonState(gvr::ControllerButton::GVR_CONTROLLER_BUTTON_RESERVED0))
+			{
+				CurrentButtonStates[EGoogleVRControllerButton::TriggerPress] = true;
+			}
+			else if (CachedControllerState.GetButtonUp(gvr::ControllerButton::GVR_CONTROLLER_BUTTON_RESERVED0))
+			{
+				CurrentButtonStates[EGoogleVRControllerButton::TriggerPress] = false;
+			}
+
+			if (CachedControllerState.GetButtonState(gvr::ControllerButton::GVR_CONTROLLER_BUTTON_RESERVED1))
+			{
+				CurrentButtonStates[EGoogleVRControllerButton::Grip] = true;
+			}
+			else if (CachedControllerState.GetButtonUp(gvr::ControllerButton::GVR_CONTROLLER_BUTTON_RESERVED1))
+			{
+				CurrentButtonStates[EGoogleVRControllerButton::Grip] = false;
+			}
 
 			// Process touches and analog information
 			// OnDown
@@ -477,46 +533,46 @@ void FGoogleVRController::ProcessControllerButtons()
 		}
 	}
 
-	MessageHandler->OnControllerAnalog( FGamepadKeyNames::MotionController_Left_Thumbstick_X, 0, TranslatedLocation.X);
-	MessageHandler->OnControllerAnalog( FGamepadKeyNames::MotionController_Right_Thumbstick_X, 0, TranslatedLocation.X);
-
-	MessageHandler->OnControllerAnalog( FGamepadKeyNames::MotionController_Left_Thumbstick_Y, 0, TranslatedLocation.Y);
-	MessageHandler->OnControllerAnalog( FGamepadKeyNames::MotionController_Right_Thumbstick_Y, 0, TranslatedLocation.Y);
+	if (Hand == EControllerHand::Left)
+	{
+		MessageHandler->OnControllerAnalog( FGamepadKeyNames::MotionController_Left_Thumbstick_X, 0, TranslatedLocation.X);
+		MessageHandler->OnControllerAnalog(FGamepadKeyNames::MotionController_Left_Thumbstick_Y, 0, TranslatedLocation.Y);
+	}
+	if (Hand == EControllerHand::Right)
+	{
+		MessageHandler->OnControllerAnalog(FGamepadKeyNames::MotionController_Right_Thumbstick_X, 0, TranslatedLocation.X);
+		MessageHandler->OnControllerAnalog(FGamepadKeyNames::MotionController_Right_Thumbstick_Y, 0, TranslatedLocation.Y);
+	}
 
-	// Process buttons for both hands at the same time
+	// Process buttons
+	const int32 HandIndex = (int32)Hand;
+	check(HandIndex >= 0 && HandIndex < CONTROLLERS_PER_PLAYER);
 	for(int32 ButtonIndex = 0; ButtonIndex < (int32)EGoogleVRControllerButton::TotalButtonCount; ++ButtonIndex)
 	{
-		if(CurrentButtonStates[ButtonIndex] != LastButtonStates[ButtonIndex])
+		if(CurrentButtonStates[ButtonIndex] != LastButtonStates[ControllerStateIndex][ButtonIndex])
 		{
 			// OnDown
 			if(CurrentButtonStates[ButtonIndex])
 			{
-				if (ButtonIndex != EGoogleVRControllerButton::TouchPadTouch)
-				{
-					MessageHandler->OnControllerButtonPressed(Buttons[(int32)EControllerHand::Left][ButtonIndex], 0, false);
-				}
-				MessageHandler->OnControllerButtonPressed( Buttons[(int32)EControllerHand::Right][ButtonIndex], 0, false);
+				MessageHandler->OnControllerButtonPressed( Buttons[HandIndex][ButtonIndex], 0, false);
 			}
 			// On Up
 			else
 			{
-				if (ButtonIndex != EGoogleVRControllerButton::TouchPadTouch)
-				{
-					MessageHandler->OnControllerButtonReleased(Buttons[(int32)EControllerHand::Left][ButtonIndex], 0, false);
-				}
-				MessageHandler->OnControllerButtonReleased( Buttons[(int32)EControllerHand::Right][ButtonIndex], 0, false);
+				MessageHandler->OnControllerButtonReleased( Buttons[HandIndex][ButtonIndex], 0, false);
 			}
 		}
 
 		// update state for next time
-		LastButtonStates[ButtonIndex] = CurrentButtonStates[ButtonIndex];
+		LastButtonStates[ControllerStateIndex][ButtonIndex] = CurrentButtonStates[ButtonIndex];
 	}
 #endif // GOOGLEVRCONTROLLER_SUPPORTED_PLATFORMS
 }
 
-void FGoogleVRController::ProcessControllerEvents()
+void FGoogleVRController::ProcessControllerEvents(int32 ControllerStateIndex)
 {
 #if GOOGLEVRCONTROLLER_SUPPORTED_PLATFORMS
+	const gvr::ControllerState& CachedControllerState = CachedControllerStates[ControllerStateIndex];
 	if (CachedControllerState.GetRecentered())
 	{
 #if GOOGLEVRCONTROLLER_SUPPORTED_EMULATOR_PLATFORMS
@@ -535,36 +591,39 @@ void FGoogleVRController::ProcessControllerEvents()
 	}
 
 
-	EGoogleVRControllerState PreviousControllerState = CurrentControllerState;
+	EGoogleVRControllerState PreviousControllerState = CurrentControllerStates[ControllerStateIndex];
 #if GOOGLEVRCONTROLLER_SUPPORTED_INSTANT_PREVIEW_PLATFORMS
 	if (gvr::ControllerConnectionState::GVR_CONTROLLER_DISCONNECTED != InstantPreviewControllerState.connection_state) {
-		CurrentControllerState = (EGoogleVRControllerState)InstantPreviewControllerState.connection_state;
+		CurrentControllerStates[ControllerStateIndex] = (EGoogleVRControllerState)InstantPreviewControllerState.connection_state;
 	}
 	else
 #endif  // GOOGLEVRCONTROLLER_SUPPORTED_INSTANT_PREVIEW_PLATFORMS
 	{
-		CurrentControllerState = (EGoogleVRControllerState) CachedControllerState.GetConnectionState();
+		CurrentControllerStates[ControllerStateIndex] = (EGoogleVRControllerState) CachedControllerState.GetConnectionState();
 	}
-	if (CurrentControllerState != PreviousControllerState)
+	if (CurrentControllerStates[ControllerStateIndex] != PreviousControllerState)
 	{
-		UGoogleVRControllerFunctionLibrary::GetGoogleVRControllerEventManager()->OnControllerStateChangedDelegate.Broadcast(CurrentControllerState);
+		UGoogleVRControllerFunctionLibrary::GetGoogleVRControllerEventManager()->OnControllerStateChangedDelegate.Broadcast(CurrentControllerStates[ControllerStateIndex]);
 	}
 #endif
 }
 
-bool FGoogleVRController::IsAvailable() const
+bool FGoogleVRController::IsAvailable(int32 ControllerStateIndex) const
 {
 #if GOOGLEVRCONTROLLER_SUPPORTED_PLATFORMS
 #if GOOGLEVRCONTROLLER_SUPPORTED_INSTANT_PREVIEW_PLATFORMS
 	if (gvr::ControllerConnectionState::GVR_CONTROLLER_CONNECTED == InstantPreviewControllerState.connection_state) {
-		return true;
+		return ControllerStateIndex == 0;
 	}
 	else
 #endif  // GOOGLEVRCONTROLLER_SUPPORTED_INSTANT_PREVIEW_PLATFORMS
-	if(CachedControllerState.GetApiStatus() == gvr::ControllerApiStatus::GVR_CONTROLLER_API_OK &&
-		CachedControllerState.GetConnectionState() == gvr::ControllerConnectionState::GVR_CONTROLLER_CONNECTED )
 	{
-		return true;
+		const gvr::ControllerState& ControllerState = CachedControllerStates[ControllerStateIndex];
+		if (ControllerState.GetApiStatus() == gvr::ControllerApiStatus::GVR_CONTROLLER_API_OK &&
+			ControllerState.GetConnectionState() == gvr::ControllerConnectionState::GVR_CONTROLLER_CONNECTED)
+		{
+			return true;
+		}
 	}
 #endif
 	return false;
@@ -589,9 +648,24 @@ int FGoogleVRController::GetGVRControllerHandedness() const
 #endif
 }
 
-EGoogleVRControllerState FGoogleVRController::GetControllerState() const
+EGoogleVRControllerAPIStatus FGoogleVRController::GetApiStatus() const
 {
-	return CurrentControllerState;
+#if GOOGLEVRCONTROLLER_SUPPORTED_ANDROID_PLATFORMS
+	// Any controller would give the same result.
+	return (EGoogleVRControllerAPIStatus) CachedControllerStates[0].GetApiStatus();
+#else
+	return EGoogleVRControllerAPIStatus::Unknown;
+#endif
+}
+
+EGoogleVRControllerState FGoogleVRController::GetControllerState(EControllerHand Hand) const
+{
+	int32 ControllerIndex = GetControllerStateIndex(Hand);
+	if (ControllerIndex >= 0)
+	{
+		return CurrentControllerStates[ControllerIndex];
+	}
+	return EGoogleVRControllerState::Disconnected;
 }
 
 FVector FGoogleVRController::ConvertGvrVectorToUnreal(float x, float y, float z, float WorldToMetersScale) const
@@ -615,50 +689,62 @@ FQuat FGoogleVRController::ConvertGvrQuaternionToUnreal(float w, float x, float
 	return Result;
 }
 
-bool FGoogleVRController::GetBatteryCharging()
+bool FGoogleVRController::GetBatteryCharging(EControllerHand Hand)
 {
 #if GOOGLEVRCONTROLLER_SUPPORTED_ANDROID_PLATFORMS
-	return CachedControllerState.GetBatteryCharging();
-#else
-	return 0;
+	const gvr::ControllerState* ControllerState = GetCachedControllerState(Hand);
+	if (ControllerState)
+	{
+		return ControllerState->GetBatteryCharging();
+	}
 #endif
+
+	return false;
 }
 
-EGoogleVRControllerBatteryLevel FGoogleVRController::GetBatteryLevel()
+EGoogleVRControllerBatteryLevel FGoogleVRController::GetBatteryLevel(EControllerHand Hand)
 {
 #if GOOGLEVRCONTROLLER_SUPPORTED_ANDROID_PLATFORMS
-	switch (CachedControllerState.GetBatteryLevel())
+	const gvr::ControllerState* ControllerState = GetCachedControllerState(Hand);
+	if (ControllerState)
 	{
-		case kControllerBatteryLevelCriticalLow:
-			return EGoogleVRControllerBatteryLevel::CriticalLow;
+		switch (ControllerState->GetBatteryLevel())
+		{
+			case kControllerBatteryLevelCriticalLow:
+				return EGoogleVRControllerBatteryLevel::CriticalLow;
 
-		case kControllerBatteryLevelLow:
-			return EGoogleVRControllerBatteryLevel::Low;
+			case kControllerBatteryLevelLow:
+				return EGoogleVRControllerBatteryLevel::Low;
 
-		case kControllerBatteryLevelMedium:
-			return EGoogleVRControllerBatteryLevel::Medium;
+			case kControllerBatteryLevelMedium:
+				return EGoogleVRControllerBatteryLevel::Medium;
 
-		case kControllerBatteryLevelAlmostFull:
-			return EGoogleVRControllerBatteryLevel::AlmostFull;
+			case kControllerBatteryLevelAlmostFull:
+				return EGoogleVRControllerBatteryLevel::AlmostFull;
 
-		case kControllerBatteryLevelFull:
-			return EGoogleVRControllerBatteryLevel::Full;
+			case kControllerBatteryLevelFull:
+				return EGoogleVRControllerBatteryLevel::Full;
 
-		default:
-			break;
-    }
+			default:
+				break;
+		}
+	}
 #endif
 
 	return EGoogleVRControllerBatteryLevel::Unknown;
 }
 
-int64_t FGoogleVRController::GetLastBatteryTimestamp()
+int64_t FGoogleVRController::GetLastBatteryTimestamp(EControllerHand Hand)
 {
 #if GOOGLEVRCONTROLLER_SUPPORTED_ANDROID_PLATFORMS
-	return CachedControllerState.GetLastBatteryTimestamp();
-#else
-	return 0;
+	const gvr::ControllerState* ControllerState = GetCachedControllerState(Hand);
+	if (ControllerState)
+	{
+		return ControllerState->GetLastBatteryTimestamp();
+	}
 #endif
+
+	return 0;
 }
 
 bool FGoogleVRController::GetUseArmModel() const
@@ -698,8 +784,11 @@ void FGoogleVRController::Tick(float DeltaTime)
 
 void FGoogleVRController::SendControllerEvents()
 {
-	ProcessControllerButtons();
-	ProcessControllerEvents();
+	for (int32 i = 0; i < CONTROLLERS_PER_PLAYER; ++i)
+	{
+		ProcessControllerButtons(i);
+		ProcessControllerEvents(i);
+	}
 }
 
 void FGoogleVRController::SetMessageHandler(const TSharedRef< FGenericApplicationMessageHandler >& InMessageHandler)
@@ -729,14 +818,37 @@ bool FGoogleVRController::GetControllerOrientationAndPosition(const int32 Contro
 	{
 		return false;
 	}
-
-	if(IsAvailable())
+	if(IsAvailable(GetControllerStateIndex(DeviceHand)))
 	{
 		OutPosition = FVector::ZeroVector;
 		OutOrientation = FRotator::ZeroRotator;
 
 #if GOOGLEVRCONTROLLER_SUPPORTED_PLATFORMS
-		if (bUseArmModel)
+
+		// Determine if we should actually use the arm model, or if we have 6dof controllers.
+		bool bUseArmModelHere = bUseArmModel;
+#if !GOOGLEVRCONTROLLER_SUPPORTED_INSTANT_PREVIEW_PLATFORMS
+#if GOOGLEVRCONTROLLER_SUPPORTED_ANDROID_PLATFORMS
+		if (bUseArmModelHere && GoogleVRCaps == EGoogleVRCaps::Daydream66)
+		{
+			// HACK: checking for reported position of exactly 0,0,0 and assuming that means no position tracking.  A future API revision will hopefully provide this somehow.
+			// We have to check both controllers, because one could be tracked while the other is still connecting or something, and we don't want to use the arm model for that one.
+			bool b6DOFControllers = false;
+			for (int32 i = 0; i < CONTROLLERS_PER_PLAYER; i++)
+			{
+				const gvr_vec3f ControllerPosition = CachedControllerStates[i].GetPosition();
+				if (ControllerPosition.x != 0.0f || ControllerPosition.y != 0.0f || ControllerPosition.z != 0.0f)
+				{
+					b6DOFControllers = true;
+					break;
+				}
+			}
+			bUseArmModelHere &= !b6DOFControllers;
+		}
+#endif //GOOGLEVRCONTROLLER_SUPPORTED_ANDROID_PLATFORMS
+#endif
+
+		if (bUseArmModelHere)
 		{
 			const gvr_arm_model::Vector3& ControllerPosition = ArmModelController.GetControllerPosition();
 			const gvr_arm_model::Quaternion& ControllerRotation = ArmModelController.GetControllerRotation();
@@ -769,8 +881,15 @@ bool FGoogleVRController::GetControllerOrientationAndPosition(const int32 Contro
 			else
 #endif  // GOOGLEVRCONTROLLER_SUPPORTED_INSTANT_PREVIEW_PLATFORMS
 			{
-				gvr_quatf ControllerOrientation = CachedControllerState.GetOrientation();
+				const gvr::ControllerState* ControllerState = GetCachedControllerState(DeviceHand);
+				check(ControllerState);  // IsAvailable check above should ensure this.
+				gvr_quatf ControllerOrientation = ControllerState->GetOrientation();
 				OutOrientation = FQuat(ControllerOrientation.qz, -ControllerOrientation.qx, -ControllerOrientation.qy, ControllerOrientation.qw).Rotator();
+#if GOOGLEVRCONTROLLER_SUPPORTED_ANDROID_PLATFORMS
+// GetPosition does not link for the editor, I guess it doesnt' exist in the windows lib.  May be fixed in future api.
+				const gvr_vec3f ControllerPosition = ControllerState->GetPosition();
+				OutPosition = ConvertGvrVectorToUnreal(ControllerPosition.x, ControllerPosition.y, ControllerPosition.z, WorldToMetersScale);
+#endif //GOOGLEVRCONTROLLER_SUPPORTED_ANDROID_PLATFORMS
 			}
 		}
 #if GOOGLEVRCONTROLLER_SUPPORTED_EMULATOR_PLATFORMS
@@ -786,6 +905,50 @@ bool FGoogleVRController::GetControllerOrientationAndPosition(const int32 Contro
 	return false;
 }
 
+int32 FGoogleVRController::GetControllerStateIndex(EControllerHand Hand) const
+{
+	switch (Hand)
+	{
+	case EControllerHand::Left:
+		return 1;
+	case EControllerHand::Right:
+	case EControllerHand::AnyHand:
+		return 0;
+	default:
+		return -1;
+	}
+}
+
+EControllerHand FGoogleVRController::GetControllerHandFromStateIndex(int32 StateIndex) const
+{
+	switch (StateIndex)
+	{
+	case 1:
+		return EControllerHand::Left;
+	case 0:
+		return EControllerHand::Right;
+	default:
+		check(false);
+		return EControllerHand::Right;
+	}
+}
+
+#if GOOGLEVRCONTROLLER_SUPPORTED_PLATFORMS
+gvr::ControllerState* FGoogleVRController::GetCachedControllerState(EControllerHand Hand)
+{
+	const int32 i = GetControllerStateIndex(Hand);
+	check(i < 2);
+	return (i >= 0) ? &(CachedControllerStates[i]) : nullptr;
+}
+
+const gvr::ControllerState* FGoogleVRController::GetCachedControllerState(EControllerHand Hand) const
+{
+	const int32 i = GetControllerStateIndex(Hand);
+	check(i < 2);
+	return (i >= 0) ? &(CachedControllerStates[i]) : nullptr;
+}
+#endif
+
 float FGoogleVRController::GetWorldToMetersScale() const
 {
 	if (IsInGameThread() && GWorld != nullptr)
@@ -800,7 +963,7 @@ float FGoogleVRController::GetWorldToMetersScale() const
 ETrackingStatus FGoogleVRController::GetControllerTrackingStatus(const int32 ControllerIndex, const EControllerHand DeviceHand) const
 {
 #if GOOGLEVRCONTROLLER_SUPPORTED_PLATFORMS
-	if(IsAvailable())
+	if(IsAvailable(GetControllerStateIndex(DeviceHand)))
 	{
 		return ETrackingStatus::Tracked;
 	}
diff --git a/Engine/Plugins/Runtime/GoogleVR/GoogleVRController/Source/GoogleVRController/Private/GoogleVRController.h b/Engine/Plugins/Runtime/GoogleVR/GoogleVRController/Source/GoogleVRController/Private/GoogleVRController.h
--- a/Engine/Plugins/Runtime/GoogleVR/GoogleVRController/Source/GoogleVRController/Private/GoogleVRController.h
+++ b/Engine/Plugins/Runtime/GoogleVR/GoogleVRController/Source/GoogleVRController/Private/GoogleVRController.h
@@ -5,6 +5,9 @@
 #include "GoogleVRControllerPrivate.h"
 #include "IInputDevice.h"
 #include "XRMotionControllerBase.h"
+#if GOOGLEVRCONTROLLER_SUPPORTED_ANDROID_PLATFORMS
+#include "AndroidRuntimeSettings.h"
+#endif
 #if GOOGLEVRCONTROLLER_SUPPORTED_PLATFORMS
 #include "gvr_arm_model.h"
 #endif
@@ -84,30 +87,31 @@ public: // Helper Functions
 	void PollController(float DeltaTime);
 
 	/** Processes the controller buttons */
-	void ProcessControllerButtons();
+	void ProcessControllerButtons(int32 ControllerStateIndex);
 
 	/** Processes the controller events */
-	void ProcessControllerEvents();
+	void ProcessControllerEvents(int32 ControllerStateIndex);
 
 	/** Checks if the controller is available */
-	bool IsAvailable() const;
+	bool IsAvailable(int32 ControllerStateIndex) const;
 
 	int GetGVRControllerHandedness() const;
 
-	EGoogleVRControllerState GetControllerState() const;
+	EGoogleVRControllerAPIStatus GetApiStatus() const;
+	EGoogleVRControllerState GetControllerState(EControllerHand Hand) const;
 
 	FVector ConvertGvrVectorToUnreal(float x, float y, float z, float WorldToMetersScale) const;
 
 	FQuat ConvertGvrQuaternionToUnreal(float w, float x, float y, float z) const;
 
 	/** Checks if the controller battery is charging. */
-	bool GetBatteryCharging();
+	bool GetBatteryCharging(EControllerHand Hand);
 
 	/** Returns an approximate battery level. */
-	EGoogleVRControllerBatteryLevel GetBatteryLevel();
+	EGoogleVRControllerBatteryLevel GetBatteryLevel(EControllerHand Hand);
 
 	/** Returns the time stamp the battery information was last updated. */
-	int64_t GetLastBatteryTimestamp();
+	int64_t GetLastBatteryTimestamp(EControllerHand Hand);
 
 public: // Arm Model
 
@@ -163,9 +167,9 @@ public: // IMotionController
 	*/
 	virtual ETrackingStatus GetControllerTrackingStatus(const int32 ControllerIndex, const EControllerHand DeviceHand) const;
 
-	/** Cached controller info */
 #if GOOGLEVRCONTROLLER_SUPPORTED_PLATFORMS
-	gvr::ControllerState CachedControllerState;
+	gvr::ControllerState* GetCachedControllerState(EControllerHand Hand);
+	const gvr::ControllerState* GetCachedControllerState(EControllerHand Hand) const;
 #endif
 
 	float  GetWorldToMetersScale() const;
@@ -174,27 +178,36 @@ private:
 
 
 #if GOOGLEVRCONTROLLER_SUPPORTED_PLATFORMS
+	/** Cached controller info */
+	gvr::ControllerState CachedControllerStates[CONTROLLERS_PER_PLAYER];
+
 	/** GVR Controller client reference */
 	gvr::ControllerApi *pController;
 
 	/** Capture Button Press states */
-	bool LastButtonStates[EGoogleVRControllerButton::TotalButtonCount];
+	bool LastButtonStates[CONTROLLERS_PER_PLAYER][EGoogleVRControllerButton::TotalButtonCount];
 
-	/** Button mappings */
+	/** Button mappings.  Index is the hand, not the controller state index. */
 	FGamepadKeyNames::Type Buttons[CONTROLLERS_PER_PLAYER][EGoogleVRControllerButton::TotalButtonCount];
 #endif
-	bool bControllerReadyToPollState;
 
 	/** handler to send all messages to */
 	TSharedRef<FGenericApplicationMessageHandler> MessageHandler;
 
+
 	// TODO: Set controller handedness based on the value from gvr.
 	// When does the gvr value become available? Should we just poll for it?
 #if GOOGLEVRCONTROLLER_SUPPORTED_PLATFORMS
 	gvr_arm_model::Controller ArmModelController;
+#endif
+#if GOOGLEVRCONTROLLER_SUPPORTED_ANDROID_PLATFORMS
+	EGoogleVRCaps::Type GoogleVRCaps;
 #endif
 	bool bUseArmModel;
-	EGoogleVRControllerState CurrentControllerState;
+
+	int32 GetControllerStateIndex(const EControllerHand Hand) const;
+	EControllerHand GetControllerHandFromStateIndex(int32 StateIndex) const;
+	EGoogleVRControllerState CurrentControllerStates[CONTROLLERS_PER_PLAYER];
 
 #if GOOGLEVRCONTROLLER_SUPPORTED_EMULATOR_PLATFORMS
 	FRotator BaseEmulatorOrientation;
diff --git a/Engine/Plugins/Runtime/GoogleVR/GoogleVRController/Source/GoogleVRController/Private/GoogleVRControllerFunctionLibrary.cpp b/Engine/Plugins/Runtime/GoogleVR/GoogleVRController/Source/GoogleVRController/Private/GoogleVRControllerFunctionLibrary.cpp
--- a/Engine/Plugins/Runtime/GoogleVR/GoogleVRController/Source/GoogleVRController/Private/GoogleVRControllerFunctionLibrary.cpp
+++ b/Engine/Plugins/Runtime/GoogleVR/GoogleVRController/Source/GoogleVRController/Private/GoogleVRControllerFunctionLibrary.cpp
@@ -30,18 +30,18 @@ EGoogleVRControllerAPIStatus UGoogleVRControllerFunctionLibrary::GetGoogleVRCont
 	FGoogleVRController* GVRController = GetGoogleVRController();
 	if(GVRController != nullptr)
 	{
-		return (EGoogleVRControllerAPIStatus) GVRController->CachedControllerState.GetApiStatus();
+		return (EGoogleVRControllerAPIStatus) GVRController->GetApiStatus();
 	}
 #endif
 	return EGoogleVRControllerAPIStatus::Unknown;
 }
 
-EGoogleVRControllerState UGoogleVRControllerFunctionLibrary::GetGoogleVRControllerState()
+EGoogleVRControllerState UGoogleVRControllerFunctionLibrary::GetGoogleVRControllerState(EControllerHand Hand)
 {
 	FGoogleVRController* GVRController = GetGoogleVRController();
 	if(GVRController != nullptr)
 	{
-		return GVRController->GetControllerState();
+		return GVRController->GetControllerState(Hand);
 	}
 
 	return EGoogleVRControllerState::Disconnected;
@@ -66,42 +66,50 @@ EGoogleVRControllerHandedness UGoogleVRControllerFunctionLibrary::GetGoogleVRCon
 	return HandednessPrefsEnum;
 }
 
-FVector UGoogleVRControllerFunctionLibrary::GetGoogleVRControllerRawAccel()
+FVector UGoogleVRControllerFunctionLibrary::GetGoogleVRControllerRawAccel(EControllerHand Hand)
 {
 #if GOOGLEVRCONTROLLER_SUPPORTED_PLATFORMS
 	FGoogleVRController* GVRController = GetGoogleVRController();
 	if(GVRController != nullptr)
 	{
-		gvr_vec3f ControllerAccel = GVRController->CachedControllerState.GetAccel();
-		return FVector(ControllerAccel.x, ControllerAccel.y, ControllerAccel.z);
+		gvr::ControllerState* ControllerState = GVRController->GetCachedControllerState(Hand);
+		if (ControllerState)
+		{
+			gvr_vec3f ControllerAccel = ControllerState->GetAccel();
+			return FVector(ControllerAccel.x, ControllerAccel.y, ControllerAccel.z);
+		}
 	}
 #endif
 
 	return FVector::ZeroVector;
 }
 
-FVector UGoogleVRControllerFunctionLibrary::GetGoogleVRControllerRawGyro()
+FVector UGoogleVRControllerFunctionLibrary::GetGoogleVRControllerRawGyro(EControllerHand Hand)
 {
 #if GOOGLEVRCONTROLLER_SUPPORTED_PLATFORMS
 	FGoogleVRController* GVRController = GetGoogleVRController();
 	if(GVRController != nullptr)
 	{
-		gvr_vec3f ControllerGyro = GVRController->CachedControllerState.GetGyro();
-		return FVector(ControllerGyro.x, ControllerGyro.y, ControllerGyro.z);
+		gvr::ControllerState* ControllerState = GVRController->GetCachedControllerState(Hand);
+		if (ControllerState)
+		{
+			gvr_vec3f ControllerGyro = ControllerState->GetGyro();
+			return FVector(ControllerGyro.x, ControllerGyro.y, ControllerGyro.z);
+		}
 	}
 #endif
 
 	return FVector::ZeroVector;
 }
 
-FRotator UGoogleVRControllerFunctionLibrary::GetGoogleVRControllerOrientation()
+FRotator UGoogleVRControllerFunctionLibrary::GetGoogleVRControllerOrientation(EControllerHand Hand)
 {
 	FGoogleVRController* GVRController = GetGoogleVRController();
 	if(GVRController != nullptr)
 	{
 		FRotator orientation;
 		FVector position;
-		GVRController->GetControllerOrientationAndPosition(0, EControllerHand::AnyHand, orientation, position, GVRController->GetWorldToMetersScale());
+		GVRController->GetControllerOrientationAndPosition(0, Hand, orientation, position, GVRController->GetWorldToMetersScale());
 		return orientation;
 	}
 	return FRotator::ZeroRotator;
@@ -416,39 +424,39 @@ float UGoogleVRControllerFunctionLibrary::GetTooltipAlphaValue()
 	return 0.0f;
 }
 
-bool UGoogleVRControllerFunctionLibrary::GetBatteryCharging()
+bool UGoogleVRControllerFunctionLibrary::GetBatteryCharging(EControllerHand Hand)
 {
 #if GOOGLEVRCONTROLLER_SUPPORTED_PLATFORMS
 	FGoogleVRController* GVRController = GetGoogleVRController();
 	if(GVRController != nullptr)
 	{
-		return GVRController->GetBatteryCharging();
+		return GVRController->GetBatteryCharging(Hand);
 	}
 #endif
 
 	return false;
 }
 
-EGoogleVRControllerBatteryLevel UGoogleVRControllerFunctionLibrary::GetBatteryLevel()
+EGoogleVRControllerBatteryLevel UGoogleVRControllerFunctionLibrary::GetBatteryLevel(EControllerHand Hand)
 {
 #if GOOGLEVRCONTROLLER_SUPPORTED_PLATFORMS
 	FGoogleVRController* GVRController = GetGoogleVRController();
 	if(GVRController != nullptr)
 	{
-		return GVRController->GetBatteryLevel();
+		return GVRController->GetBatteryLevel(Hand);
 	}
 #endif
 
 	return EGoogleVRControllerBatteryLevel::Unknown;
 }
 
-int64_t UGoogleVRControllerFunctionLibrary::GetLastBatteryTimestamp()
+int64_t UGoogleVRControllerFunctionLibrary::GetLastBatteryTimestamp(EControllerHand Hand)
 {
 #if GOOGLEVRCONTROLLER_SUPPORTED_PLATFORMS
 	FGoogleVRController* GVRController = GetGoogleVRController();
 	if(GVRController != nullptr)
 	{
-		return GVRController->GetLastBatteryTimestamp();
+		return GVRController->GetLastBatteryTimestamp(Hand);
 	}
 #endif
 
diff --git a/Engine/Plugins/Runtime/GoogleVR/GoogleVRHMD/Source/GoogleVRHMD/GoogleVRHMD_APL.xml b/Engine/Plugins/Runtime/GoogleVR/GoogleVRHMD/Source/GoogleVRHMD/GoogleVRHMD_APL.xml
--- a/Engine/Plugins/Runtime/GoogleVR/GoogleVRHMD/Source/GoogleVRHMD/GoogleVRHMD_APL.xml
+++ b/Engine/Plugins/Runtime/GoogleVR/GoogleVRHMD/Source/GoogleVRHMD/GoogleVRHMD_APL.xml
@@ -24,8 +24,10 @@
 		<setBoolFromPropertyContains result="bCardboardCaps" ini="Engine" section="/Script/AndroidRuntimeSettings.AndroidRuntimeSettings" property="GoogleVRCaps" contains="Cardboard"/>
 		<setBoolFromPropertyContains result="bDaydream33Caps" ini="Engine" section="/Script/AndroidRuntimeSettings.AndroidRuntimeSettings" property="GoogleVRCaps" contains="Daydream33"/>
 		<setBoolFromPropertyContains result="bDaydream63Caps" ini="Engine" section="/Script/AndroidRuntimeSettings.AndroidRuntimeSettings" property="GoogleVRCaps" contains="Daydream63"/>
+		<setBoolFromPropertyContains result="bDaydream66Caps" ini="Engine" section="/Script/AndroidRuntimeSettings.AndroidRuntimeSettings" property="GoogleVRCaps" contains="Daydream66"/>
 
 		<setBoolOr result="bSupportDaydream" arg1="$B(bDaydream33Caps)" arg2="$B(bDaydream63Caps)"/>
+		<setBoolOr result="bSupportDaydream" arg1="$B(bSupportDaydream)" arg2="$B(bDaydream66Caps)"/>
 		<if condition="bSupportDaydream">
 			<true>
 				<log text="Project supports daydream, will enable async reprojection when possible." />
@@ -56,6 +58,11 @@
 				<log text="Project supports Daydream AIO 6.3 device." />
 			</true>
 		</if>
+		<if condition="bDaydream66Caps">
+			<true>
+				<log text="Project supports Daydream AIO 6.6 device." />
+			</true>
+		</if>
 	</init>
 
 	<!-- optional files or directories to copy to Intermediate/Android/APK -->
@@ -146,16 +153,19 @@
 				<addFeature android:name="android.software.vr.mode" android:required="false" />
 			</false>
 		</if>
+		<if condition="bDaydream33Caps">
+			<true>
+				<addFeature android:name="android.hardware.vr.headtracking" android:version="1" android:required="false" />
+			</true>
+		</if>
 		<if condition="bDaydream63Caps">
 			<true>
-				<if condition="bDaydream33Caps">
-					<true>
-						<addFeature android:name="android.hardware.vr.headtracking" android:version="1" android:required="false" />
-					</true>
-					<false>
-						<addFeature android:name="android.hardware.vr.headtracking" android:version="1" android:required="true" />
-					</false>
-				</if>
+				<addFeature android:name="android.hardware.vr.headtracking" android:version="1" android:required="true" />
+			</true>
+		</if>
+		<if condition="bDaydream66Caps">
+			<true>
+				<addFeature android:name="android.hardware.vr.headtracking" android:version="1" android:required="true" />
 			</true>
 		</if>
 
diff --git a/Engine/Plugins/Runtime/GoogleVR/GoogleVRHMD/Source/GoogleVRHMD/Private/GoogleVRHMD.cpp b/Engine/Plugins/Runtime/GoogleVR/GoogleVRHMD/Source/GoogleVRHMD/Private/GoogleVRHMD.cpp
--- a/Engine/Plugins/Runtime/GoogleVR/GoogleVRHMD/Source/GoogleVRHMD/Private/GoogleVRHMD.cpp
+++ b/Engine/Plugins/Runtime/GoogleVR/GoogleVRHMD/Source/GoogleVRHMD/Private/GoogleVRHMD.cpp
@@ -397,6 +397,12 @@ FGoogleVRHMD::FGoogleVRHMD(const FAutoRegister& AutoRegister)
 
 #if GOOGLEVRHMD_SUPPORTED_ANDROID_PLATFORMS
 
+	// set to identity
+	CachedHeadPose = 	{ { { 1.0f, 0.0f, 0.0f, 0.0f },
+						{ 0.0f, 1.0f, 0.0f, 0.0f },
+						{ 0.0f, 0.0f, 1.0f, 0.0f },
+						{ 0.0f, 0.0f, 0.0f, 1.0f } } };
+
 	// Get GVRAPI from java
 	GVRAPI = AndroidThunkCpp_GetNativeGVRApi();
 
diff --git a/Engine/Plugins/Runtime/GoogleVR/GoogleVRHMD/Source/GoogleVRHMD/Private/GoogleVRHMDRender.cpp b/Engine/Plugins/Runtime/GoogleVR/GoogleVRHMD/Source/GoogleVRHMD/Private/GoogleVRHMDRender.cpp
--- a/Engine/Plugins/Runtime/GoogleVR/GoogleVRHMD/Source/GoogleVRHMD/Private/GoogleVRHMDRender.cpp
+++ b/Engine/Plugins/Runtime/GoogleVR/GoogleVRHMD/Source/GoogleVRHMD/Private/GoogleVRHMDRender.cpp
@@ -418,6 +418,12 @@ FGoogleVRHMDCustomPresent::FGoogleVRHMDCustomPresent(FGoogleVRHMD* InHMD)
 	, CurrentFrameViewportList(nullptr)
 	, bSkipPresent(false)
 {
+	// set to identity
+	CurrentFrameRenderHeadPose =	{ { { 1.0f, 0.0f, 0.0f, 0.0f },
+									{ 0.0f, 1.0f, 0.0f, 0.0f },
+									{ 0.0f, 0.0f, 1.0f, 0.0f },
+									{ 0.0f, 0.0f, 0.0f, 1.0f } } };
+
 	CreateGVRSwapChain();
 }
 
diff --git a/Engine/Source/Developer/Android/AndroidPlatformEditor/Private/AndroidTargetSettingsCustomization.cpp b/Engine/Source/Developer/Android/AndroidPlatformEditor/Private/AndroidTargetSettingsCustomization.cpp
--- a/Engine/Source/Developer/Android/AndroidPlatformEditor/Private/AndroidTargetSettingsCustomization.cpp
+++ b/Engine/Source/Developer/Android/AndroidPlatformEditor/Private/AndroidTargetSettingsCustomization.cpp
@@ -386,7 +386,7 @@ void FAndroidTargetSettingsCustomization::OnGoogleVRCapsChange()
 
 	const TArray<TEnumAsByte<EGoogleVRCaps::Type>> &GoogleCaps = GetDefault<UAndroidRuntimeSettings>()->GoogleVRCaps;
 
-	bool bIsDaydream = GoogleCaps.Contains(EGoogleVRCaps::Daydream33) || GoogleCaps.Contains(EGoogleVRCaps::Daydream63);
+	bool bIsDaydream = GoogleCaps.Contains(EGoogleVRCaps::Daydream33) || GoogleCaps.Contains(EGoogleVRCaps::Daydream63) || GoogleCaps.Contains(EGoogleVRCaps::Daydream66);
 	if (bIsDaydream && GetDefault<UAndroidRuntimeSettings>()->bAllowIMU)
 	{
 		// turn off IMU for Daydream (but user can turn it back on
diff --git a/Engine/Source/Programs/UnrealBuildTool/Platform/Android/UEDeployAndroid.cs b/Engine/Source/Programs/UnrealBuildTool/Platform/Android/UEDeployAndroid.cs
--- a/Engine/Source/Programs/UnrealBuildTool/Platform/Android/UEDeployAndroid.cs
+++ b/Engine/Source/Programs/UnrealBuildTool/Platform/Android/UEDeployAndroid.cs
@@ -432,7 +432,7 @@ namespace UnrealBuildTool
 			List<string> GoogleVRCaps = new List<string>();
 			if(Ini.GetArray("/Script/AndroidRuntimeSettings.AndroidRuntimeSettings", "GoogleVRCaps", out GoogleVRCaps))
 			{
-				return GoogleVRCaps.Contains("Daydream33") || GoogleVRCaps.Contains("Daydream63");
+				return GoogleVRCaps.Contains("Daydream33") || GoogleVRCaps.Contains("Daydream63") || GoogleVRCaps.Contains("Daydream66");
 			}
 			else
 			{
diff --git a/Engine/Source/Runtime/Android/AndroidRuntimeSettings/Classes/AndroidRuntimeSettings.h b/Engine/Source/Runtime/Android/AndroidRuntimeSettings/Classes/AndroidRuntimeSettings.h
--- a/Engine/Source/Runtime/Android/AndroidRuntimeSettings/Classes/AndroidRuntimeSettings.h
+++ b/Engine/Source/Runtime/Android/AndroidRuntimeSettings/Classes/AndroidRuntimeSettings.h
@@ -154,7 +154,8 @@ namespace EGoogleVRCaps
 	{
 		Cardboard = 0 UMETA(DisplayName = "Cardboard", ToolTip = "Head orientation, no controller."),
 		Daydream33 = 1 UMETA(DisplayName = "Daydream (3.3 DoF)", ToolTip = "Head orientation, controller orientation. Daydream without positional tracking."),
-		Daydream63 = 2 UMETA(DisplayName = "Daydream (6.3 DoF)", ToolTip = "Head position and orientation, controller orientation. Daydream with positional tracking.")
+		Daydream63 = 2 UMETA(DisplayName = "Daydream (6.3 DoF)", ToolTip = "Head position and orientation, controller orientation. Daydream with positional tracking."),
+		Daydream66 = 3 UMETA(DisplayName = "Daydream (6.6 DoF)", ToolTip = "Head position and orientation, 2 controllers with position and orientation. Daydream with positional tracking.")
 	};
 }
 
