Disable GoogleVRHMD since Linux library is missing anyway.

- GoogleInstantPreview Linux library is missing, making the rest of the integration invalid due to dangling symbols.
- Fix UE-60644 if re-enabling.

diff --git a/Engine/Plugins/Runtime/GoogleVR/GoogleVRController/GoogleVRController.uplugin b/Engine/Plugins/Runtime/GoogleVR/GoogleVRController/GoogleVRController.uplugin
--- a/Engine/Plugins/Runtime/GoogleVR/GoogleVRController/GoogleVRController.uplugin
+++ b/Engine/Plugins/Runtime/GoogleVR/GoogleVRController/GoogleVRController.uplugin
@@ -20,7 +20,7 @@
             "Name": "GoogleVRController",
             "Type": "Runtime",
             "LoadingPhase": "PostConfigInit",
-            "WhitelistPlatforms": [ "Android", "Mac", "Win32", "Win64", "Linux" ]
+            "WhitelistPlatforms": [ "Android", "Mac", "Win32", "Win64" ]
         }
     ],
     "Plugins": [
diff --git a/Engine/Plugins/Runtime/GoogleVR/GoogleVRHMD/GoogleVRHMD.uplugin b/Engine/Plugins/Runtime/GoogleVR/GoogleVRHMD/GoogleVRHMD.uplugin
--- a/Engine/Plugins/Runtime/GoogleVR/GoogleVRHMD/GoogleVRHMD.uplugin
+++ b/Engine/Plugins/Runtime/GoogleVR/GoogleVRHMD/GoogleVRHMD.uplugin
@@ -20,7 +20,7 @@
             "Name": "GoogleVRHMD",
             "Type": "Runtime",
             "LoadingPhase": "PostConfigInit",
-            "WhitelistPlatforms": [ "Win64", "Win32", "Mac", "Linux", "Android" ]
+            "WhitelistPlatforms": [ "Win64", "Win32", "Mac", "Android" ]
         }
     ]
 }
