Fixing build break in GoogleVRControllerFunctionLibrary.h

diff --git a/Engine/Plugins/Runtime/GoogleVR/GoogleVRController/Source/GoogleVRController/Classes/GoogleVRControllerFunctionLibrary.h b/Engine/Plugins/Runtime/GoogleVR/GoogleVRController/Source/GoogleVRController/Classes/GoogleVRControllerFunctionLibrary.h
--- a/Engine/Plugins/Runtime/GoogleVR/GoogleVRController/Source/GoogleVRController/Classes/GoogleVRControllerFunctionLibrary.h
+++ b/Engine/Plugins/Runtime/GoogleVR/GoogleVRController/Source/GoogleVRController/Classes/GoogleVRControllerFunctionLibrary.h
@@ -4,6 +4,7 @@
 
 #include "GoogleVRControllerEventManager.h"
 #include "Kismet/BlueprintFunctionLibrary.h"
+#include "InputCoreTypes.h"
 #include "GoogleVRControllerFunctionLibrary.generated.h"
 
 UENUM(BlueprintType)
