Fixing GoogleVRController build warning

diff --git a/Engine/Plugins/Runtime/GoogleVR/GoogleVRController/Source/GoogleVRController/Private/GoogleVRController.cpp b/Engine/Plugins/Runtime/GoogleVR/GoogleVRController/Source/GoogleVRController/Private/GoogleVRController.cpp
--- a/Engine/Plugins/Runtime/GoogleVR/GoogleVRController/Source/GoogleVRController/Private/GoogleVRController.cpp
+++ b/Engine/Plugins/Runtime/GoogleVR/GoogleVRController/Source/GoogleVRController/Private/GoogleVRController.cpp
@@ -273,34 +273,36 @@ void FGoogleVRController::PollController(float DeltaTime)
 #endif
 
 #if GOOGLEVRCONTROLLER_SUPPORTED_EMULATOR_PLATFORMS
-	gvr::ControllerState* ControllerState = GetCachedControllerState(EControllerHand::AnyHand);
-	check(ControllerState);
-	int32_t PreviousConnectionState = ControllerState->GetConnectionState();
-	// If controller connection is requested but it is not connected, try resetup adb forward.
-	if(bKeepConnectingControllerEmulator
-	   && PreviousConnectionState != gvr::ControllerConnectionState::GVR_CONTROLLER_CONNECTED)
-	{
-		double CurrentTime = FPlatformTime::Seconds();
-		if(CurrentTime - LastTimeTryAdbForward > ADB_FORWARD_RETRY_TIME)
+	{
+		gvr::ControllerState* ControllerState = GetCachedControllerState(EControllerHand::AnyHand);
+		check(ControllerState);
+		int32_t PreviousConnectionState = ControllerState->GetConnectionState();
+		// If controller connection is requested but it is not connected, try resetup adb forward.
+		if (bKeepConnectingControllerEmulator
+			&& PreviousConnectionState != gvr::ControllerConnectionState::GVR_CONTROLLER_CONNECTED)
 		{
-			UE_LOG(LogGoogleVRController, Log, TEXT("Trying to connect to GoogleVR Controller"));
-			SetupAdbForward();
-			LastTimeTryAdbForward = CurrentTime;
+			double CurrentTime = FPlatformTime::Seconds();
+			if (CurrentTime - LastTimeTryAdbForward > ADB_FORWARD_RETRY_TIME)
+			{
+				UE_LOG(LogGoogleVRController, Log, TEXT("Trying to connect to GoogleVR Controller"));
+				SetupAdbForward();
+				LastTimeTryAdbForward = CurrentTime;
+			}
 		}
-	}
 
-	ControllerState->Update(*pController);
+		ControllerState->Update(*pController);
 
-	if(PreviousConnectionState != gvr::ControllerConnectionState::GVR_CONTROLLER_CONNECTED
-	   && ControllerState->GetConnectionState() == gvr::ControllerConnectionState::GVR_CONTROLLER_CONNECTED)
-	{
-		UE_LOG(LogGoogleVRController, Log, TEXT("GoogleVR Controller Connected"));
-	}
+		if (PreviousConnectionState != gvr::ControllerConnectionState::GVR_CONTROLLER_CONNECTED
+			&& ControllerState->GetConnectionState() == gvr::ControllerConnectionState::GVR_CONTROLLER_CONNECTED)
+		{
+			UE_LOG(LogGoogleVRController, Log, TEXT("GoogleVR Controller Connected"));
+		}
 
-	if(PreviousConnectionState == gvr::ControllerConnectionState::GVR_CONTROLLER_CONNECTED
-	   && ControllerState->GetConnectionState() != gvr::ControllerConnectionState::GVR_CONTROLLER_CONNECTED)
-	{
-		UE_LOG(LogGoogleVRController, Log, TEXT("GoogleVR Controller Disconnected"));
+		if (PreviousConnectionState == gvr::ControllerConnectionState::GVR_CONTROLLER_CONNECTED
+			&& ControllerState->GetConnectionState() != gvr::ControllerConnectionState::GVR_CONTROLLER_CONNECTED)
+		{
+			UE_LOG(LogGoogleVRController, Log, TEXT("GoogleVR Controller Disconnected"));
+		}
 	}
 #elif GOOGLEVRCONTROLLER_SUPPORTED_ANDROID_PLATFORMS
 	for (int32 i = 0; i < CONTROLLERS_PER_PLAYER; i++)
