Fixing AndroidRuntimeSettings build errors.

diff --git a/Engine/Plugins/Runtime/GoogleVR/GoogleVRController/Source/GoogleVRController/GoogleVRController.Build.cs b/Engine/Plugins/Runtime/GoogleVR/GoogleVRController/Source/GoogleVRController/GoogleVRController.Build.cs
--- a/Engine/Plugins/Runtime/GoogleVR/GoogleVRController/Source/GoogleVRController/GoogleVRController.Build.cs
+++ b/Engine/Plugins/Runtime/GoogleVR/GoogleVRController/Source/GoogleVRController/GoogleVRController.Build.cs
@@ -35,13 +35,18 @@ namespace UnrealBuildTool.Rules
 					"UMG",
 					"Slate",
 					"SlateCore",
-					"AndroidRuntimeSettings"
 				}
 				);
 
 			if (Target.Platform == UnrealTargetPlatform.Android)
 			{
-				PrivateDependencyModuleNames.Add("Launch");
+				PrivateDependencyModuleNames.AddRange(
+					new string[]
+					{
+						"Launch",
+						"AndroidRuntimeSettings",
+					}
+					);
 			}
 
 			if (Target.bBuildEditor == true)
diff --git a/Engine/Plugins/Runtime/GoogleVR/GoogleVRController/Source/GoogleVRController/Private/GoogleVRController.cpp b/Engine/Plugins/Runtime/GoogleVR/GoogleVRController/Source/GoogleVRController/Private/GoogleVRController.cpp
--- a/Engine/Plugins/Runtime/GoogleVR/GoogleVRController/Source/GoogleVRController/Private/GoogleVRController.cpp
+++ b/Engine/Plugins/Runtime/GoogleVR/GoogleVRController/Source/GoogleVRController/Private/GoogleVRController.cpp
@@ -143,6 +143,7 @@ FGoogleVRController::FGoogleVRController(const TSharedRef< FGenericApplicationMe
 		CurrentControllerStates[i] = EGoogleVRControllerState::Disconnected;
 	}
 
+#if GOOGLEVRCONTROLLER_SUPPORTED_ANDROID_PLATFORMS
 	{
 		GoogleVRCaps = EGoogleVRCaps::Cardboard;
 		FString ValueString;
@@ -157,6 +158,7 @@ FGoogleVRController::FGoogleVRController(const TSharedRef< FGenericApplicationMe
 			}
 		}
 	}
+#endif
 
 #if GOOGLEVRCONTROLLER_SUPPORTED_PLATFORMS
 	// Register motion controller!
@@ -822,10 +824,10 @@ bool FGoogleVRController::GetControllerOrientationAndPosition(const int32 Contro
 
 		// Determine if we should actually use the arm model, or if we have 6dof controllers.
 		bool bUseArmModelHere = bUseArmModel;
-		if (bUseArmModelHere && GoogleVRCaps == EGoogleVRCaps::Daydream66)
-		{
 #if !GOOGLEVRCONTROLLER_SUPPORTED_INSTANT_PREVIEW_PLATFORMS
 #if GOOGLEVRCONTROLLER_SUPPORTED_ANDROID_PLATFORMS
+		if (bUseArmModelHere && GoogleVRCaps == EGoogleVRCaps::Daydream66)
+		{
 			// HACK: checking for reported position of exactly 0,0,0 and assuming that means no position tracking.  A future API revision will hopefully provide this somehow.
 			// We have to check both controllers, because one could be tracked while the other is still connecting or something, and we don't want to use the arm model for that one.
 			bool b6DOFControllers = false;
@@ -839,9 +841,9 @@ bool FGoogleVRController::GetControllerOrientationAndPosition(const int32 Contro
 				}
 			}
 			bUseArmModelHere &= !b6DOFControllers;
+		}
 #endif //GOOGLEVRCONTROLLER_SUPPORTED_ANDROID_PLATFORMS
 #endif
-		}
 
 		if (bUseArmModelHere)
 		{
diff --git a/Engine/Plugins/Runtime/GoogleVR/GoogleVRController/Source/GoogleVRController/Private/GoogleVRController.h b/Engine/Plugins/Runtime/GoogleVR/GoogleVRController/Source/GoogleVRController/Private/GoogleVRController.h
--- a/Engine/Plugins/Runtime/GoogleVR/GoogleVRController/Source/GoogleVRController/Private/GoogleVRController.h
+++ b/Engine/Plugins/Runtime/GoogleVR/GoogleVRController/Source/GoogleVRController/Private/GoogleVRController.h
@@ -5,7 +5,9 @@
 #include "GoogleVRControllerPrivate.h"
 #include "IInputDevice.h"
 #include "XRMotionControllerBase.h"
+#if GOOGLEVRCONTROLLER_SUPPORTED_ANDROID_PLATFORMS
 #include "AndroidRuntimeSettings.h"
+#endif
 #if GOOGLEVRCONTROLLER_SUPPORTED_PLATFORMS
 #include "gvr_arm_model.h"
 #endif
@@ -198,7 +200,9 @@ private:
 #if GOOGLEVRCONTROLLER_SUPPORTED_PLATFORMS
 	gvr_arm_model::Controller ArmModelController;
 #endif
+#if GOOGLEVRCONTROLLER_SUPPORTED_ANDROID_PLATFORMS
 	EGoogleVRCaps::Type GoogleVRCaps;
+#endif
 	bool bUseArmModel;
 
 	int32 GetControllerStateIndex(const EControllerHand Hand) const;
