Fix for static analysis warning

diff --git a/Engine/Plugins/Runtime/GoogleVR/GoogleVRController/Source/GoogleVRController/Private/GoogleVRController.cpp b/Engine/Plugins/Runtime/GoogleVR/GoogleVRController/Source/GoogleVRController/Private/GoogleVRController.cpp
--- a/Engine/Plugins/Runtime/GoogleVR/GoogleVRController/Source/GoogleVRController/Private/GoogleVRController.cpp
+++ b/Engine/Plugins/Runtime/GoogleVR/GoogleVRController/Source/GoogleVRController/Private/GoogleVRController.cpp
@@ -404,6 +404,7 @@ void FGoogleVRController::ProcessControllerButtons(int32 ControllerStateIndex)
 {
 #if GOOGLEVRCONTROLLER_SUPPORTED_PLATFORMS
 	const EControllerHand Hand = GetControllerHandFromStateIndex(ControllerStateIndex);
+	check(Hand == EControllerHand::Left || Hand == EControllerHand::Right);
 
 	// Capture our current button states
 	bool CurrentButtonStates[EGoogleVRControllerButton::TotalButtonCount] = {0};
@@ -543,7 +544,9 @@ void FGoogleVRController::ProcessControllerButtons(int32 ControllerStateIndex)
 		MessageHandler->OnControllerAnalog(FGamepadKeyNames::MotionController_Right_Thumbstick_Y, 0, TranslatedLocation.Y);
 	}
 
-	// Process buttons for both hands at the same time
+	// Process buttons
+	const int32 HandIndex = (int32)Hand;
+	check(HandIndex > 0 && HandIndex < CONTROLLERS_PER_PLAYER);
 	for(int32 ButtonIndex = 0; ButtonIndex < (int32)EGoogleVRControllerButton::TotalButtonCount; ++ButtonIndex)
 	{
 		if(CurrentButtonStates[ButtonIndex] != LastButtonStates[ControllerStateIndex][ButtonIndex])
@@ -551,12 +554,12 @@ void FGoogleVRController::ProcessControllerButtons(int32 ControllerStateIndex)
 			// OnDown
 			if(CurrentButtonStates[ButtonIndex])
 			{
-				MessageHandler->OnControllerButtonPressed( Buttons[(int32)Hand][ButtonIndex], 0, false);
+				MessageHandler->OnControllerButtonPressed( Buttons[HandIndex][ButtonIndex], 0, false);
 			}
 			// On Up
 			else
 			{
-				MessageHandler->OnControllerButtonReleased( Buttons[(int32)Hand][ButtonIndex], 0, false);
+				MessageHandler->OnControllerButtonReleased( Buttons[HandIndex][ButtonIndex], 0, false);
 			}
 		}
 
