Build fix for bad assert

diff --git a/Engine/Plugins/Runtime/GoogleVR/GoogleVRController/Source/GoogleVRController/Private/GoogleVRController.cpp b/Engine/Plugins/Runtime/GoogleVR/GoogleVRController/Source/GoogleVRController/Private/GoogleVRController.cpp
--- a/Engine/Plugins/Runtime/GoogleVR/GoogleVRController/Source/GoogleVRController/Private/GoogleVRController.cpp
+++ b/Engine/Plugins/Runtime/GoogleVR/GoogleVRController/Source/GoogleVRController/Private/GoogleVRController.cpp
@@ -546,7 +546,7 @@ void FGoogleVRController::ProcessControllerButtons(int32 ControllerStateIndex)
 
 	// Process buttons
 	const int32 HandIndex = (int32)Hand;
-	check(HandIndex > 0 && HandIndex < CONTROLLERS_PER_PLAYER);
+	check(HandIndex >= 0 && HandIndex < CONTROLLERS_PER_PLAYER);
 	for(int32 ButtonIndex = 0; ButtonIndex < (int32)EGoogleVRControllerButton::TotalButtonCount; ++ButtonIndex)
 	{
 		if(CurrentButtonStates[ButtonIndex] != LastButtonStates[ControllerStateIndex][ButtonIndex])
