GVR only supports windows and mac for emulation.

diff --git a/Engine/Plugins/Runtime/GoogleVR/GoogleVRController/Source/GoogleVRController/Public/IGoogleVRControllerPlugin.h b/Engine/Plugins/Runtime/GoogleVR/GoogleVRController/Source/GoogleVRController/Public/IGoogleVRControllerPlugin.h
--- a/Engine/Plugins/Runtime/GoogleVR/GoogleVRController/Source/GoogleVRController/Public/IGoogleVRControllerPlugin.h
+++ b/Engine/Plugins/Runtime/GoogleVR/GoogleVRController/Source/GoogleVRController/Public/IGoogleVRControllerPlugin.h
@@ -6,7 +6,7 @@
 #include "IInputDeviceModule.h"
 
 #define GOOGLEVRCONTROLLER_SUPPORTED_ANDROID_PLATFORMS (PLATFORM_ANDROID)
-#define GOOGLEVRCONTROLLER_SUPPORTED_EMULATOR_PLATFORMS (WITH_EDITOR)
+#define GOOGLEVRCONTROLLER_SUPPORTED_EMULATOR_PLATFORMS (WITH_EDITOR && (PLATFORM_WINDOWS || PLATFORM_MAC))
 #define GOOGLEVRCONTROLLER_SUPPORTED_INSTANT_PREVIEW_PLATFORMS (WITH_EDITOR)
 #define GOOGLEVRCONTROLLER_SUPPORTED_PLATFORMS (GOOGLEVRCONTROLLER_SUPPORTED_ANDROID_PLATFORMS || GOOGLEVRCONTROLLER_SUPPORTED_EMULATOR_PLATFORMS || GOOGLEVRCONTROLLER_SUPPORTED_INSTANT_PREVIEW_PLATFORMS)
 
